Improve path tracing denoising results behind sharp reflections and refractions

Use the current path roughness and material roughness to decide how much of the current material albedo and normal should be written to the denoiser outputs. This improves texture fidelity behind glass and in reflections substantially after denoising.

Allow the denoiser plugin to read the denoiser cvar to experiment with different denoising behaviors. I am intending this to remain an undocumented feature, the default behavior still appears to be the best for now.

diff --git a/Engine/Plugins/Experimental/OpenImageDenoise/Source/OpenImageDenoise/Private/Plugin.cpp b/Engine/Plugins/Experimental/OpenImageDenoise/Source/OpenImageDenoise/Private/Plugin.cpp
--- a/Engine/Plugins/Experimental/OpenImageDenoise/Source/OpenImageDenoise/Private/Plugin.cpp
+++ b/Engine/Plugins/Experimental/OpenImageDenoise/Source/OpenImageDenoise/Private/Plugin.cpp
@@ -24,9 +24,19 @@ DEFINE_LOG_CATEGORY(LogOpenImageDenoise);
 
 IMPLEMENT_MODULE(FOpenImageDenoiseModule, OpenImageDenoise)
 
+enum class EDenoiseMode {
+	OFF = 0,
+	DEFAULT = 1,   // denoise using albedo/normal as provided
+	CLEAN_AUX = 2, // denoise albedo/normal first and pass cleanAux (haven't found a scene where this is clearly better, and it runs much slower)
+	NO_AOVS = 3,   // denoise beauty only
+};
+
 static void Denoise(FRHICommandListImmediate& RHICmdList, FRHITexture2D* ColorTex, FRHITexture2D* AlbedoTex, FRHITexture2D* NormalTex, FRHITexture2D* OutputTex)
 {
-	const int DenoiserMode = 2; // TODO: Expose setting for this
+	static IConsoleVariable* DenoiseModeCVar = IConsoleManager::Get().FindConsoleVariable(TEXT("r.PathTracing.Denoiser"));
+
+	const int DenoiseModeCVarValue = DenoiseModeCVar ? DenoiseModeCVar->GetInt() : -1;
+	const EDenoiseMode DenoiseMode = DenoiseModeCVarValue >= 0 ? EDenoiseMode(DenoiseModeCVarValue) : EDenoiseMode::DEFAULT;
 
 #if WITH_EDITOR
 	uint64 FilterExecuteTime = 0;
@@ -41,7 +51,7 @@ static void Denoise(FRHICommandListImmediate& RHICmdList, FRHITexture2D* ColorTe
 	FReadSurfaceDataFlags ReadDataFlags(ERangeCompressionMode::RCM_MinMax);
 	ReadDataFlags.SetLinearToGamma(false);
 	RHICmdList.ReadSurfaceData(ColorTex, Rect, RawPixels, ReadDataFlags);
-	if (DenoiserMode >= 2)
+	if (DenoiseMode == EDenoiseMode::DEFAULT || DenoiseMode == EDenoiseMode::CLEAN_AUX)
 	{
 		RHICmdList.ReadSurfaceData(AlbedoTex, Rect, RawAlbedo, ReadDataFlags);
 		RHICmdList.ReadSurfaceData(NormalTex, Rect, RawNormal, ReadDataFlags);
@@ -49,23 +59,45 @@ static void Denoise(FRHICommandListImmediate& RHICmdList, FRHITexture2D* ColorTe
 
 	check(RawPixels.Num() == Size.X * Size.Y);
 
-	uint32_t DestStride;
-	FLinearColor* DestBuffer = (FLinearColor*)RHICmdList.LockTexture2D(OutputTex, 0, RLM_WriteOnly, DestStride, false);
-
 	// create device only once?
 	oidn::DeviceRef OIDNDevice = oidn::newDevice();
 	OIDNDevice.commit();
+
+	if (DenoiseMode == EDenoiseMode::CLEAN_AUX)
+	{
+		oidn::FilterRef OIDNFilter = OIDNDevice.newFilter("RT");
+		OIDNFilter.setImage("albedo", RawAlbedo.GetData(), oidn::Format::Float3, Size.X, Size.Y, 0, sizeof(FLinearColor), sizeof(FLinearColor) * Size.X);
+		OIDNFilter.setImage("output", RawAlbedo.GetData(), oidn::Format::Float3, Size.X, Size.Y, 0, sizeof(FLinearColor), sizeof(FLinearColor) * Size.X);
+		OIDNFilter.commit();
+		OIDNFilter.execute();
+	}
+
+	if (DenoiseMode == EDenoiseMode::CLEAN_AUX)
+	{
+		oidn::FilterRef OIDNFilter = OIDNDevice.newFilter("RT");
+		OIDNFilter.setImage("normal", RawNormal.GetData(), oidn::Format::Float3, Size.X, Size.Y, 0, sizeof(FLinearColor), sizeof(FLinearColor) * Size.X);
+		OIDNFilter.setImage("output", RawNormal.GetData(), oidn::Format::Float3, Size.X, Size.Y, 0, sizeof(FLinearColor), sizeof(FLinearColor) * Size.X);
+		OIDNFilter.commit();
+		OIDNFilter.execute();
+	}
+
 	oidn::FilterRef OIDNFilter = OIDNDevice.newFilter("RT");
 	OIDNFilter.setImage("color", RawPixels.GetData(), oidn::Format::Float3, Size.X, Size.Y, 0, sizeof(FLinearColor), sizeof(FLinearColor) * Size.X);
-	if (DenoiserMode >= 2)
+	if (DenoiseMode == EDenoiseMode::DEFAULT || DenoiseMode == EDenoiseMode::CLEAN_AUX)
 	{
+		// default behavior
 		OIDNFilter.setImage("albedo", RawAlbedo.GetData(), oidn::Format::Float3, Size.X, Size.Y, 0, sizeof(FLinearColor), sizeof(FLinearColor) * Size.X);
 		OIDNFilter.setImage("normal", RawNormal.GetData(), oidn::Format::Float3, Size.X, Size.Y, 0, sizeof(FLinearColor), sizeof(FLinearColor) * Size.X);
 	}
-	if (DenoiserMode >= 3)
+	if (DenoiseMode == EDenoiseMode::CLEAN_AUX)
 	{
+		// +cleanAux
 		OIDNFilter.set("cleanAux", true);
 	}
+
+	uint32_t DestStride;
+	FLinearColor* DestBuffer = (FLinearColor*)RHICmdList.LockTexture2D(OutputTex, 0, RLM_WriteOnly, DestStride, false);
+
 	OIDNFilter.setImage("output", DestBuffer, oidn::Format::Float3, Size.X, Size.Y, 0, sizeof(FLinearColor), DestStride);
 	OIDNFilter.set("hdr", true);
 	OIDNFilter.commit();
diff --git a/Engine/Shaders/Private/PathTracing/Material/PathTracingMaterialSampling.ush b/Engine/Shaders/Private/PathTracing/Material/PathTracingMaterialSampling.ush
--- a/Engine/Shaders/Private/PathTracing/Material/PathTracingMaterialSampling.ush
+++ b/Engine/Shaders/Private/PathTracing/Material/PathTracingMaterialSampling.ush
@@ -177,6 +177,90 @@ float3 EstimateMaterialAlbedo(FPathTracingPayload Payload)
 	return Albedo;
 }
 
+void AccumulateAlbedoNormal(inout FPathTracingPayload Payload, float3 PathThroughput, float PathRoughness, inout float3 Albedo, inout float3 Normal)
+{
+	const float MaxRoughness = 0.15; // don't bother past this level
+	const float MinRoughness = 0.10; // capture anything below this level
+	if (PathRoughness >= MaxRoughness)
+	{
+		// no point tracking past a certain roughness
+		return;
+	}
+	// figure out how much the current path should contribute to the denoiser aovs
+	PathThroughput *= 1 - saturate((PathRoughness - MinRoughness) / (MaxRoughness - MinRoughness));
+	// figure out how much the specular lobes should contribute to the total (relative to the current PathRoughness)
+	// if we are above the MinRoughness, use the PathRoughness as are lower bound instead so that multiple bounces don't accumulate too much albedo
+	const float RoughnessWeight = saturate((Payload.Roughness - max(MinRoughness, PathRoughness)) / (MaxRoughness - max(MinRoughness, PathRoughness)));
+	float BrdfWeight = 1;
+	bool HasRefraction = false;
+	switch (Payload.BlendingMode)
+	{
+		case RAY_TRACING_BLEND_MODE_TRANSLUCENT:
+#if !SIMPLIFIED_MATERIAL_SHADER
+			if (Payload.Ior != 1.0)
+			{
+				HasRefraction = true;
+			}
+#endif
+			{
+				BrdfWeight = Payload.Opacity;
+			}
+			break;
+		case RAY_TRACING_BLEND_MODE_ADDITIVE:
+			BrdfWeight = Payload.Opacity;
+			break;
+		case RAY_TRACING_BLEND_MODE_MODULATE:
+		case RAY_TRACING_BLEND_MODE_ALPHA_HOLDOUT:
+			// In these modes, the light does not interact with the material at all, only passes through
+			return;
+		default:
+			break;
+	}
+
+	switch (Payload.ShadingModelID)
+	{
+	case SHADINGMODELID_UNLIT:
+		// this object does not scatter light
+		return;
+#if SIMPLIFIED_MATERIAL_SHADER
+	case SHADINGMODELID_TWOSIDED_FOLIAGE:
+		Albedo += BrdfWeight * PathThroughput * (Payload.BaseColor + Payload.GetSubsurfaceColor());
+		Normal += Luminance(BrdfWeight * PathThroughput) * Payload.WorldNormal;
+		break;
+#else
+	case SHADINGMODELID_TWOSIDED_FOLIAGE:
+	case SHADINGMODELID_SUBSURFACE:
+	case SHADINGMODELID_PREINTEGRATED_SKIN:
+	case SHADINGMODELID_SUBSURFACE_PROFILE:
+		Albedo += BrdfWeight * PathThroughput * (Payload.DiffuseColor + RoughnessWeight * Payload.SpecularColor + Payload.GetSubsurfaceColor());
+		Normal += Luminance(BrdfWeight * PathThroughput) * Payload.WorldNormal;
+		break;
+	case SHADINGMODELID_DEFAULT_LIT:
+	case SHADINGMODELID_CLEAR_COAT:
+		Albedo += BrdfWeight * PathThroughput * (Payload.DiffuseColor + RoughnessWeight * Payload.SpecularColor);
+		Normal += Luminance(BrdfWeight * PathThroughput) * Payload.WorldNormal;
+		break;
+	case SHADINGMODELID_CLOTH:
+		Albedo += BrdfWeight * PathThroughput * (Payload.DiffuseColor + lerp(RoughnessWeight * Payload.SpecularColor, Payload.CustomData.xyz, Payload.CustomData.w));
+		Normal += Luminance(BrdfWeight * PathThroughput) * Payload.WorldNormal;
+		break;
+#endif
+	default:
+		Albedo += BrdfWeight * PathThroughput * Payload.BaseColor;
+		Normal += Luminance(BrdfWeight * PathThroughput) * Payload.WorldNormal;
+		break;
+	}
+#if !SIMPLIFIED_MATERIAL_SHADER
+	if (HasRefraction)
+	{
+		// glass is currently perfectly clear - so has unit albedo
+		// TODO: estimate some kind of albedo for thin-translucent case since that one uses a color?
+		Albedo += (1 - BrdfWeight) * RoughnessWeight * PathThroughput;
+		Normal += (1 - BrdfWeight) * RoughnessWeight * Luminance(PathThroughput) * Payload.WorldNormal;
+	}
+#endif
+}
+
 void AdjustMaterialRoughness(inout FPathTracingPayload Payload, float PathRoughness)
 {
 #if !SIMPLIFIED_MATERIAL_SHADER
diff --git a/Engine/Shaders/Private/PathTracing/PathTracingCore.ush b/Engine/Shaders/Private/PathTracing/PathTracingCore.ush
--- a/Engine/Shaders/Private/PathTracing/PathTracingCore.ush
+++ b/Engine/Shaders/Private/PathTracing/PathTracingCore.ush
@@ -66,7 +66,7 @@ void AccumulateRadiance(inout float3 TotalRadiance, float3 PathRadiance, bool bI
 	TotalRadiance += PathRadiance;
 }
 
-FPathTracingPayload TraceTransparentRay(RayDesc Ray, FRayCone RayCone, bool bIsCameraRay, bool bLastBounce, bool bIncludeEmission, uint2 LaunchIndex, uint NumLights, inout RandomSequence RandSequence, inout float3 PathThroughput, inout float3 Radiance, inout float3 Albedo, inout float3 Normal)
+FPathTracingPayload TraceTransparentRay(RayDesc Ray, FRayCone RayCone, bool bIsCameraRay, bool bLastBounce, bool bIncludeEmission, uint2 LaunchIndex, uint NumLights, float PathRoughness, inout RandomSequence RandSequence, inout float3 PathThroughput, inout float3 Radiance, inout float3 Albedo, inout float3 Normal)
 {
 	const uint RayFlags = bIsCameraRay && EnableCameraBackfaceCulling ? RAY_FLAG_CULL_BACK_FACING_TRIANGLES : 0;
 	const uint MissShaderIndex = 0;
@@ -134,11 +134,9 @@ FPathTracingPayload TraceTransparentRay(RayDesc Ray, FRayCone RayCone, bool bIsC
 		if (!bLastBounce)
 		{
 			float3 Contrib = PathThroughput * EstimateMaterialAlbedo(HitPayload);
-			if (bIsCameraRay)
-			{
-				Albedo += Contrib;
-				Normal += PathThroughput * (1 - Luminance(Transparency)) * HitPayload.WorldNormal;
-			}
+
+			// accumulate what the denoiser wants into albedo/normal (as long as the current path is rough enough)
+			AccumulateAlbedoNormal(HitPayload, PathThroughput, PathRoughness, Albedo, Normal);
 
 			float SelectionWeight = max3(Contrib.x, Contrib.y, Contrib.z);
 			SelectionWeightSum += SelectionWeight;
@@ -648,6 +646,7 @@ bool PathTracingKernel(inout FPathState PathState, int Bounce)
 		bIncludeEmissive,
 		PathState.PixelIndex,
 		NumVisibleLights,
+		PathState.PathRoughness,
 		PathState.RandSequence,
 		PathState.PathThroughput,
 		PathState.Radiance,
@@ -827,11 +826,8 @@ bool PathTracingKernel(inout FPathState PathState, int Bounce)
 	PathState.Ray.TMax = POSITIVE_INFINITY;
 	ApplyRayBias(PathState.Ray, Payload.HitT, MaterialSample.PositionBiasSign * Payload.WorldGeoNormal);
 
-	if (ApproximateCaustics)
-	{
-		// enlarge roughness based on the chosen lobe roughness
-		PathState.PathRoughness = max(PathState.PathRoughness, MaterialSample.Roughness);
-	}
+	// enlarge roughness based on the chosen lobe roughness
+	PathState.PathRoughness = max(PathState.PathRoughness, MaterialSample.Roughness);
 
 	// If we are using Material sampling for lights
 	if (MISMode != 1 && bDoLightLoop)
@@ -840,6 +836,7 @@ bool PathTracingKernel(inout FPathState PathState, int Bounce)
 		// While it would be possible to just loop around and use the indirect ray to do this, it would prevent the application
 		// of shadow ray specific logic for transparent shadows or various per light tricks like shadow casting
 		const bool bUseMIS = MISMode == 2 && LightPickingCdfSum > 0;
+		const float ShadowRayRoughness = ApproximateCaustics ? PathState.PathRoughness : 0.0;
 		for (uint Index = 0, Num = LightLoopCount.NumMISLights; Index < Num; ++Index)
 		{
 			uint LightId = GetLightId(Index, LightLoopCount);
@@ -871,7 +868,7 @@ bool PathTracingKernel(inout FPathState PathState, int Bounce)
 				{
 					RayDesc LightRay = PathState.Ray;
 					LightRay.TMax = LightResult.HitT;
-					LightContrib *= TraceTransparentVisibilityRay(LightRay, PathState.PixelIndex, PathState.PathRoughness);
+					LightContrib *= TraceTransparentVisibilityRay(LightRay, PathState.PixelIndex, ShadowRayRoughness);
 				}
 				// the light made some contribution, and there was nothing along the shadow ray
 				AccumulateRadiance(PathState.Radiance, LightContrib, bIsCameraRay);
