Cleanup old references to RHITexture2D

diff --git a/Engine/Plugins/Experimental/OpenImageDenoise/Source/OpenImageDenoise/Private/Plugin.cpp b/Engine/Plugins/Experimental/OpenImageDenoise/Source/OpenImageDenoise/Private/Plugin.cpp
--- a/Engine/Plugins/Experimental/OpenImageDenoise/Source/OpenImageDenoise/Private/Plugin.cpp
+++ b/Engine/Plugins/Experimental/OpenImageDenoise/Source/OpenImageDenoise/Private/Plugin.cpp
@@ -31,7 +31,7 @@ enum class EDenoiseMode {
 	NO_AOVS = 3,   // denoise beauty only
 };
 
-static void Denoise(FRHICommandListImmediate& RHICmdList, FRHITexture2D* ColorTex, FRHITexture2D* AlbedoTex, FRHITexture2D* NormalTex, FRHITexture2D* OutputTex, FRHIGPUMask GPUMask)
+static void Denoise(FRHICommandListImmediate& RHICmdList, FRHITexture* ColorTex, FRHITexture* AlbedoTex, FRHITexture* NormalTex, FRHITexture* OutputTex, FRHIGPUMask GPUMask)
 {
 	static IConsoleVariable* DenoiseModeCVar = IConsoleManager::Get().FindConsoleVariable(TEXT("r.PathTracing.Denoiser"));
 
diff --git a/Engine/Source/Runtime/Renderer/Public/PathTracingDenoiser.h b/Engine/Source/Runtime/Renderer/Public/PathTracingDenoiser.h
--- a/Engine/Source/Runtime/Renderer/Public/PathTracingDenoiser.h
+++ b/Engine/Source/Runtime/Renderer/Public/PathTracingDenoiser.h
@@ -5,6 +5,6 @@
 class FRHITexture;
 class FRHICommandListImmediate;
 
-using PathTracingDenoiserFunction = void(FRHICommandListImmediate& RHICmdList, FRHITexture2D* ColorTex, FRHITexture2D* AlbedoTex, FRHITexture2D* NormalTex, FRHITexture2D* OutputTex, FRHIGPUMask GPUMask);
+using PathTracingDenoiserFunction = void(FRHICommandListImmediate& RHICmdList, FRHITexture* ColorTex, FRHITexture* AlbedoTex, FRHITexture* NormalTex, FRHITexture* OutputTex, FRHIGPUMask GPUMask);
 
 extern RENDERER_API PathTracingDenoiserFunction* GPathTracingDenoiserFunc;
