OpenImageDenoise: Fix crash on shutdown if denoiser was never used

Refactor code to avoid the static singleton and instead incorporate the denoiser state into the object allocated from GPathTracingDenoiserPlugin.

#rnx

diff --git a/Engine/Plugins/Experimental/OpenImageDenoise/Source/OpenImageDenoise/Private/Plugin.cpp b/Engine/Plugins/Experimental/OpenImageDenoise/Source/OpenImageDenoise/Private/Plugin.cpp
--- a/Engine/Plugins/Experimental/OpenImageDenoise/Source/OpenImageDenoise/Private/Plugin.cpp
+++ b/Engine/Plugins/Experimental/OpenImageDenoise/Source/OpenImageDenoise/Private/Plugin.cpp
@@ -173,15 +173,8 @@ struct OIDNState
 			}
 		}
 	}
-	
-	void Reset()
-	{
-		UpdateFilter(FIntPoint(0, 0), FDenoiseSettings{});
-	}
 };
 
-static OIDNState DenoiserState;
-
 template <typename PixelType>
 static void CopyTextureFromGPUToCPU(FRHICommandListImmediate& RHICmdList, FRHITexture* SrcTexture, FIntPoint Size, TArray<PixelType>& DstArray)
 {
@@ -215,7 +208,7 @@ static void CopyTextureFromCPUToGPU(FRHICommandListImmediate& RHICmdList, const
 	RHICmdList.UnlockTexture2D(DstTexture, 0, false);
 }
 
-static void Denoise(FRHICommandListImmediate& RHICmdList, FRHITexture* ColorTex, FRHITexture* AlbedoTex, FRHITexture* NormalTex, FRHITexture* OutputTex, FRHIGPUMask GPUMask)
+static void Denoise(OIDNState& DenoiserState, FRHICommandListImmediate& RHICmdList, FRHITexture* ColorTex, FRHITexture* AlbedoTex, FRHITexture* NormalTex, FRHITexture* OutputTex, FRHIGPUMask GPUMask)
 {
 	FDenoiseSettings Settings = GetCurrentSettings();
 
@@ -273,6 +266,8 @@ using namespace UE::Renderer::Private;
 
 class FOIDNDenoiser : public IPathTracingDenoiser
 {
+	mutable OIDNState DenoiserState;
+
 public:
 	~FOIDNDenoiser() {}
 
@@ -286,9 +281,10 @@ public:
 
 		// Need to read GPU mask outside Pass function, as the value is not refreshed inside the pass
 		GraphBuilder.AddPass(RDG_EVENT_NAME("OIDN Denoiser Plugin"), DenoiseParameters, ERDGPassFlags::Readback,
-			[DenoiseParameters, GPUMask = View.GPUMask](FRHICommandListImmediate& RHICmdList)
+			[DenoiseParameters, GPUMask = View.GPUMask, this](FRHICommandListImmediate& RHICmdList)
 		{
-			Denoise(RHICmdList,
+			Denoise(DenoiserState,
+				RHICmdList,
 				DenoiseParameters->InputTexture->GetRHI()->GetTexture2D(),
 				DenoiseParameters->InputAlbedo->GetRHI()->GetTexture2D(),
 				DenoiseParameters->InputNormal->GetRHI()->GetTexture2D(),
@@ -314,6 +310,5 @@ void FOpenImageDenoiseModule::ShutdownModule()
 #endif
 
 	// Release scratch memory and destroy the OIDN device and filters
-	DenoiserState.Reset();
 	GPathTracingDenoiserPlugin.Reset();
 }
