Path Tracer: Change the radiance buffer back to 32 bit precision

Half precision is not sufficient for accumulating many samples and could lead to banding artifacts on smooth gradients in some cases.

#rnx

diff --git a/Engine/Plugins/Experimental/OpenImageDenoise/Source/OpenImageDenoise/Private/Plugin.cpp b/Engine/Plugins/Experimental/OpenImageDenoise/Source/OpenImageDenoise/Private/Plugin.cpp
--- a/Engine/Plugins/Experimental/OpenImageDenoise/Source/OpenImageDenoise/Private/Plugin.cpp
+++ b/Engine/Plugins/Experimental/OpenImageDenoise/Source/OpenImageDenoise/Private/Plugin.cpp
@@ -54,11 +54,16 @@ static TAutoConsoleVariable<bool> CVarOIDNUseAux(
 	TEXT("Should OpenImageDenoise make use of auxilary buffers (albedo and normal) to improve image quality? (default: true)")
 );
 
-struct FHalfColor
+struct FColor16
 {
 	uint16_t r, g, b, a;
 };
 
+struct FColor32
+{
+	float r, g, b, a;
+};
+
 struct FDenoiseSettings
 {
 	bool DenoiseAlpha;
@@ -82,9 +87,9 @@ FDenoiseSettings GetCurrentSettings()
 struct OIDNState
 {
 	// scratch CPU memory for running the OIDN filter
-	TArray<FHalfColor> RawPixels;
-	TArray<FHalfColor> RawAlbedo;
-	TArray<FHalfColor> RawNormal;
+	TArray<FColor32> RawPixels;
+	TArray<FColor16> RawAlbedo;
+	TArray<FColor16> RawNormal;
 
 	// re-useable filters
 	oidn::FilterRef AlbedoFilter;
@@ -128,13 +133,13 @@ struct OIDNState
 			if (CurrentSettings.DenoiseAux)
 			{
 				AlbedoFilter = OIDNDevice.newFilter("RT");
-				AlbedoFilter.setImage("albedo", RawAlbedo.GetData(), oidn::Format::Half3, Size.X, Size.Y, 0, sizeof(FHalfColor), sizeof(FHalfColor) * Size.X);
-				AlbedoFilter.setImage("output", RawAlbedo.GetData(), oidn::Format::Half3, Size.X, Size.Y, 0, sizeof(FHalfColor), sizeof(FHalfColor) * Size.X);
+				AlbedoFilter.setImage("albedo", RawAlbedo.GetData(), oidn::Format::Half3, Size.X, Size.Y, 0, sizeof(FColor16), sizeof(FColor16) * Size.X);
+				AlbedoFilter.setImage("output", RawAlbedo.GetData(), oidn::Format::Half3, Size.X, Size.Y, 0, sizeof(FColor16), sizeof(FColor16) * Size.X);
 				AlbedoFilter.set("quality", oidn::Quality::High);
 				AlbedoFilter.commit();
 				NormalFilter = OIDNDevice.newFilter("RT");
-				NormalFilter.setImage("normal", RawNormal.GetData(), oidn::Format::Half3, Size.X, Size.Y, 0, sizeof(FHalfColor), sizeof(FHalfColor) * Size.X);
-				NormalFilter.setImage("output", RawNormal.GetData(), oidn::Format::Half3, Size.X, Size.Y, 0, sizeof(FHalfColor), sizeof(FHalfColor) * Size.X);
+				NormalFilter.setImage("normal", RawNormal.GetData(), oidn::Format::Half3, Size.X, Size.Y, 0, sizeof(FColor16), sizeof(FColor16) * Size.X);
+				NormalFilter.setImage("output", RawNormal.GetData(), oidn::Format::Half3, Size.X, Size.Y, 0, sizeof(FColor16), sizeof(FColor16) * Size.X);
 				NormalFilter.set("quality", oidn::Quality::High);
 				NormalFilter.commit();
 			}
@@ -144,14 +149,14 @@ struct OIDNState
 				NormalFilter = oidn::FilterRef();
 			}
 			PixelsFilter = OIDNDevice.newFilter("RT");
-			PixelsFilter.setImage("color" , RawPixels.GetData(), oidn::Format::Half3, Size.X, Size.Y, 0, sizeof(FHalfColor), sizeof(FHalfColor) * Size.X);
-			PixelsFilter.setImage("output", RawPixels.GetData(), oidn::Format::Half3, Size.X, Size.Y, 0, sizeof(FHalfColor), sizeof(FHalfColor) * Size.X);
+			PixelsFilter.setImage("color" , RawPixels.GetData(), oidn::Format::Float3, Size.X, Size.Y, 0, sizeof(FColor32), sizeof(FColor32) * Size.X);
+			PixelsFilter.setImage("output", RawPixels.GetData(), oidn::Format::Float3, Size.X, Size.Y, 0, sizeof(FColor32), sizeof(FColor32) * Size.X);
 
 			if (CurrentSettings.UseAux)
 			{
 				// default behavior, use the albedo/normal buffers to improve quality
-				PixelsFilter.setImage("albedo", RawAlbedo.GetData(), oidn::Format::Half3, Size.X, Size.Y, 0, sizeof(FHalfColor), sizeof(FHalfColor) * Size.X);
-				PixelsFilter.setImage("normal", RawNormal.GetData(), oidn::Format::Half3, Size.X, Size.Y, 0, sizeof(FHalfColor), sizeof(FHalfColor) * Size.X);
+				PixelsFilter.setImage("albedo", RawAlbedo.GetData(), oidn::Format::Half3, Size.X, Size.Y, 0, sizeof(FColor16), sizeof(FColor16) * Size.X);
+				PixelsFilter.setImage("normal", RawNormal.GetData(), oidn::Format::Half3, Size.X, Size.Y, 0, sizeof(FColor16), sizeof(FColor16) * Size.X);
 			}
 			if (CurrentSettings.DenoiseAux && CurrentSettings.UseAux)
 			{
@@ -165,8 +170,8 @@ struct OIDNState
 			if (CurrentSettings.DenoiseAlpha)
 			{
 				AlphaFilter = OIDNDevice.newFilter("RT");
-				AlphaFilter.setImage("color" , RawPixels.GetData(), oidn::Format::Half, Size.X, Size.Y, sizeof(uint16_t) * 3, sizeof(uint16_t), sizeof(FHalfColor) * Size.X);
-				AlphaFilter.setImage("output", RawPixels.GetData(), oidn::Format::Half, Size.X, Size.Y, sizeof(uint16_t) * 3, sizeof(uint16_t), sizeof(FHalfColor) * Size.X);
+				AlphaFilter.setImage("color" , RawPixels.GetData(), oidn::Format::Float, Size.X, Size.Y, sizeof(float) * 3, sizeof(float), sizeof(FColor32) * Size.X);
+				AlphaFilter.setImage("output", RawPixels.GetData(), oidn::Format::Float, Size.X, Size.Y, sizeof(float) * 3, sizeof(float), sizeof(FColor32) * Size.X);
 				AlphaFilter.set("hdr", true);
 				AlphaFilter.commit();
 			}
diff --git a/Engine/Source/Runtime/Renderer/Private/PathTracing.cpp b/Engine/Source/Runtime/Renderer/Private/PathTracing.cpp
--- a/Engine/Source/Runtime/Renderer/Private/PathTracing.cpp
+++ b/Engine/Source/Runtime/Renderer/Private/PathTracing.cpp
@@ -3183,7 +3183,7 @@ void FDeferredShadingSceneRenderer::RenderPathTracing(
 		// First time through, need to make a new texture
 		FRDGTextureDesc RadianceDesc = FRDGTextureDesc::Create2D(
 			View.ViewRect.Size(),
-			PF_FloatRGBA,
+			PF_A32B32G32R32F, // radiance accumulation must take place in floats to avoid quantization artifacts on smooth gradients
 			FClearValueBinding::None,
 			TexCreate_ShaderResource | TexCreate_UAV | GetExtraTextureCreateFlagsForDenoiser());
 		FRDGTextureDesc AlbedoNormalDesc = FRDGTextureDesc::Create2D(
