Path Tracer: Disable OpenImageDenoise plugin by default.

The path tracing denoiser now defaults to our new NNEDenoiser, so there is no need to leave OIDN enabled by default. It can still be re-enabled locally for debugging.

#rnx

diff --git a/Engine/Plugins/Experimental/OpenImageDenoise/OpenImageDenoise.uplugin b/Engine/Plugins/Experimental/OpenImageDenoise/OpenImageDenoise.uplugin
--- a/Engine/Plugins/Experimental/OpenImageDenoise/OpenImageDenoise.uplugin
+++ b/Engine/Plugins/Experimental/OpenImageDenoise/OpenImageDenoise.uplugin
@@ -10,7 +10,7 @@
 	"DocsURL" : "",
 	"MarketplaceURL" : "",
 	"SupportURL" : "",
-	"EnabledByDefault" : true,
+	"EnabledByDefault" : false,
 	"CanContainContent" : false,
 	"IsBetaVersion" : false,
 	"Installed" : false,
