OpenImageDenoise: Fix alpha channel denoising

The pixel stride for alpha denoising was wrong, the values are spaced by 16 bytes, not 4. This was leading to an incorrect vertical line artifact in the output image.

diff --git a/Engine/Plugins/Experimental/OpenImageDenoise/Source/OpenImageDenoise/Private/Plugin.cpp b/Engine/Plugins/Experimental/OpenImageDenoise/Source/OpenImageDenoise/Private/Plugin.cpp
--- a/Engine/Plugins/Experimental/OpenImageDenoise/Source/OpenImageDenoise/Private/Plugin.cpp
+++ b/Engine/Plugins/Experimental/OpenImageDenoise/Source/OpenImageDenoise/Private/Plugin.cpp
@@ -170,8 +170,8 @@ struct OIDNState
 			if (CurrentSettings.DenoiseAlpha)
 			{
 				AlphaFilter = OIDNDevice.newFilter("RT");
-				AlphaFilter.setImage("color" , RawPixels.GetData(), oidn::Format::Float, Size.X, Size.Y, sizeof(float) * 3, sizeof(float), sizeof(FColor32) * Size.X);
-				AlphaFilter.setImage("output", RawPixels.GetData(), oidn::Format::Float, Size.X, Size.Y, sizeof(float) * 3, sizeof(float), sizeof(FColor32) * Size.X);
+				AlphaFilter.setImage("color" , RawPixels.GetData(), oidn::Format::Float, Size.X, Size.Y, sizeof(float) * 3, sizeof(FColor32), sizeof(FColor32) * Size.X);
+				AlphaFilter.setImage("output", RawPixels.GetData(), oidn::Format::Float, Size.X, Size.Y, sizeof(float) * 3, sizeof(FColor32), sizeof(FColor32) * Size.X);
 				AlphaFilter.set("hdr", true);
 				AlphaFilter.commit();
 			}
