Path Tracer: Refactor UE_LOG calls to UE_LOGF and remove unused deprecated plugin interface

#rnx
#Synced-CL 51437754

diff --git a/Engine/Plugins/Experimental/OpenImageDenoise/Source/OpenImageDenoise/Private/Plugin.cpp b/Engine/Plugins/Experimental/OpenImageDenoise/Source/OpenImageDenoise/Private/Plugin.cpp
--- a/Engine/Plugins/Experimental/OpenImageDenoise/Source/OpenImageDenoise/Private/Plugin.cpp
+++ b/Engine/Plugins/Experimental/OpenImageDenoise/Source/OpenImageDenoise/Private/Plugin.cpp
@@ -124,10 +124,10 @@ struct OIDNState
 		if (!PixelsFilter || CurrentSettings != Settings || CurrentSize != Size)
 		{
 #if WITH_EDITOR
-			UE_LOG(LogOpenImageDenoise, Log, TEXT("Updating filter configuration for (%d x %d) with denoise_alpha=%s, denoise_aux=%s, use_aux=%s)"), Size.X, Size.Y,
-				   Settings.DenoiseAlpha ? TEXT("on") : TEXT("off"),
-				   Settings.DenoiseAux   ? TEXT("on") : TEXT("off"),
-				   Settings.UseAux       ? TEXT("on") : TEXT("off"));
+			UE_LOGF(LogOpenImageDenoise, Log, "Updating filter configuration for (%d x %d) with denoise_alpha=%s, denoise_aux=%s, use_aux=%s)", Size.X, Size.Y,
+				   Settings.DenoiseAlpha ? "on" : "off",
+				   Settings.DenoiseAux   ? "on" : "off",
+				   Settings.UseAux       ? "on" : "off");
 #endif
 			CurrentSettings = Settings;
 			CurrentSize = Size;
@@ -261,13 +261,13 @@ static void Denoise(OIDNState& DenoiserState, FRHICommandListImmediate& RHICmdLi
 	const char* errorMessage;
 	if (DenoiserState.OIDNDevice.getError(errorMessage) != oidn::Error::None)
 	{
-		UE_LOG(LogOpenImageDenoise, Warning, TEXT("Denoiser failed: %hs"), errorMessage);
+		UE_LOGF(LogOpenImageDenoise, Warning, "Denoiser failed: %s", errorMessage);
 		return;
 	}
 
 	FilterExecuteTime += FPlatformTime::Cycles64();
 	const double FilterExecuteTimeMS = 1000.0 * FPlatformTime::ToSeconds64(FilterExecuteTime);
-	UE_LOG(LogOpenImageDenoise, Log, TEXT("Denoised %d x %d pixels in %.2f ms"), Size.X, Size.Y, FilterExecuteTimeMS);
+	UE_LOGF(LogOpenImageDenoise, Log, "Denoised %d x %d pixels in %.2f ms", Size.X, Size.Y, FilterExecuteTimeMS);
 #endif
 }
 
@@ -306,7 +306,7 @@ public:
 void FOpenImageDenoiseModule::StartupModule()
 {
 #if WITH_EDITOR
-	UE_LOG(LogOpenImageDenoise, Log, TEXT("OIDN starting up (built with OIDN %hs)"), OIDN_VERSION_STRING);
+	UE_LOGF(LogOpenImageDenoise, Log, "OIDN starting up (built with OIDN %s", OIDN_VERSION_STRING);
 #endif
 
 	RegisterSpatialDenoiser(MakeUnique<FOIDNDenoiser>(),TEXT("OIDN"));
@@ -315,7 +315,7 @@ void FOpenImageDenoiseModule::StartupModule()
 void FOpenImageDenoiseModule::ShutdownModule()
 {
 #if WITH_EDITOR
-	UE_LOG(LogOpenImageDenoise, Log, TEXT("OIDN shutting down"));
+	UE_LOGF(LogOpenImageDenoise, Log, "OIDN shutting down");
 #endif
 
 	// Release scratch memory and destroy the OIDN device and filters
diff --git a/Engine/Source/Runtime/Renderer/Private/PathTracing.cpp b/Engine/Source/Runtime/Renderer/Private/PathTracing.cpp
--- a/Engine/Source/Runtime/Renderer/Private/PathTracing.cpp
+++ b/Engine/Source/Runtime/Renderer/Private/PathTracing.cpp
@@ -3828,7 +3828,7 @@ void PreparePathTracingRTPSO()
 				FRayTracingPipelineState* PipelineState = PipelineStateCache::GetAndOrCreateRayTracingPipelineState(RHICmdList, Initializer, ERayTracingPipelineCacheFlags::NonBlocking);
 				NumValidPermutations++;
 			}
-			UE_LOG(LogRenderer, Log, TEXT("Requested compilation of Path Tracing RTPSOs (%d permutations)."), NumValidPermutations);
+			UE_LOGF(LogRenderer, Log, "Requested compilation of Path Tracing RTPSOs (%d permutations).", NumValidPermutations);
 		}
 	);
 }
diff --git a/Engine/Source/Runtime/Renderer/Private/PathTracingSpatialTemporalDenoising.cpp b/Engine/Source/Runtime/Renderer/Private/PathTracingSpatialTemporalDenoising.cpp
--- a/Engine/Source/Runtime/Renderer/Private/PathTracingSpatialTemporalDenoising.cpp
+++ b/Engine/Source/Runtime/Renderer/Private/PathTracingSpatialTemporalDenoising.cpp
@@ -2044,7 +2044,7 @@ void PathTracingSpatialTemporalDenoising(FRDGBuilder& GraphBuilder,
 
 			if (IsSourceTargetDimensionMatch)
 			{
-				UE_LOG(LogPathTracingDenoising, Log, TEXT("Using temporal denoising for frame %i"), SpatialTemporalDenoisingContext.FrameIndex);
+				UE_LOGF(LogPathTracingDenoising, Log, "Using temporal denoising for frame %d", SpatialTemporalDenoisingContext.FrameIndex);
 
 				if (TemporalDenoiserType == ETemporalDenoiserType::BUILTIN_TEMPORAL_DENOISER)
 				{
diff --git a/Engine/Source/Runtime/Renderer/Public/PathTracingDenoiser.h b/Engine/Source/Runtime/Renderer/Public/PathTracingDenoiser.h
--- a/Engine/Source/Runtime/Renderer/Public/PathTracingDenoiser.h
+++ b/Engine/Source/Runtime/Renderer/Public/PathTracingDenoiser.h
@@ -108,12 +108,6 @@ public:
 
 } // namespace UE::Renderer::Private
 
-UE_DEPRECATED(5.5, "GPathTracingDenoiserPlugin is deprecated. Use RegisterSpatialDenoiser/RegisterSpatialTemporalDenoiser to register")
-extern RENDERER_API TUniquePtr<UE::Renderer::Private::IPathTracingDenoiser> GPathTracingDenoiserPlugin;
-
-UE_DEPRECATED(5.5, "GPathTracingSpatialTemporalDenoiserPlugin is deprecated. Use RegisterSpatialDenoiser/RegisterSpatialTemporalDenoiser to register")
-extern RENDERER_API TUniquePtr<UE::Renderer::Private::IPathTracingSpatialTemporalDenoiser> GPathTracingSpatialTemporalDenoiserPlugin;
-
 /** Register a spatial denoiser*/
 extern RENDERER_API void RegisterSpatialDenoiser(TUniquePtr<UE::Renderer::Private::IPathTracingDenoiser> PathTracingDenoiser, FString Name);
 
